Disable parallel Android texture compression on Mac due to thread safety issues with Qualcomm's Mac TextureConverter

diff --git a/Engine/Source/Developer/TextureFormatAndroid/Private/TextureFormatAndroid.cpp b/Engine/Source/Developer/TextureFormatAndroid/Private/TextureFormatAndroid.cpp
--- a/Engine/Source/Developer/TextureFormatAndroid/Private/TextureFormatAndroid.cpp
+++ b/Engine/Source/Developer/TextureFormatAndroid/Private/TextureFormatAndroid.cpp
@@ -130,7 +130,7 @@ class FTextureFormatAndroid : public ITextureFormat
 {
 	virtual bool AllowParallelBuild() const override
 	{
-		return true;
+		return !PLATFORM_MAC; // On Mac Qualcomm's TextureConverter library is not thead-safe
 	}
 
 	virtual uint16 GetVersion(
